[Mutable] Added missing Mutable Asset User Data plugin files.

- Plugin re-enabled.

#Synced-CL 51468966

diff --git a/Engine/Plugins/Experimental/MutableAssetUserData/MutableAssetUserData.uplugin b/Engine/Plugins/Experimental/MutableAssetUserData/MutableAssetUserData.uplugin
--- a/Engine/Plugins/Experimental/MutableAssetUserData/MutableAssetUserData.uplugin
+++ b/Engine/Plugins/Experimental/MutableAssetUserData/MutableAssetUserData.uplugin
@@ -14,7 +14,7 @@
 	"IsBetaVersion": false,
 	"IsExperimentalVersion": true,
 	"Installed": false,
-	"EnabledByDefault": false,
+	"EnabledByDefault": true,
 	"Modules": 
 	[
 		{
diff --git a/Engine/Plugins/Experimental/MutableAssetUserData/Source/MutableAssetUserData/Private/MutableAssetUserDataModule.cpp b/Engine/Plugins/Experimental/MutableAssetUserData/Source/MutableAssetUserData/Private/MutableAssetUserDataModule.cpp
new file mode 100644
--- /dev/null
+++ b/Engine/Plugins/Experimental/MutableAssetUserData/Source/MutableAssetUserData/Private/MutableAssetUserDataModule.cpp
@@ -0,0 +1,25 @@
+// Copyright Epic Games, Inc. All Rights Reserved.
+
+#include "MutableAssetUserDataModule.h"
+
+#include "MeshAddAssetUserData.h"
+#include "MuCO/ICustomizableObjectModule.h"
+
+
+IMPLEMENT_MODULE(FMutableAssetUserDataModule, MutableAssetUserData);
+
+
+void FMutableAssetUserDataModule::StartupModule()
+{
+	ICustomizableObjectModule& Module = ICustomizableObjectModule::Get();
+
+	Module.RegisterExternalOperation(FMeshAddAssetUserData::StaticStruct());
+}
+
+
+void FMutableAssetUserDataModule::ShutdownModule()
+{
+	ICustomizableObjectModule& Module = ICustomizableObjectModule::Get();
+
+	Module.UnregisterExternalOperation(FMeshAddAssetUserData::StaticStruct());
+}
diff --git a/Engine/Plugins/Experimental/MutableAssetUserData/Source/MutableAssetUserData/Public/MutableAssetUserDataModule.h b/Engine/Plugins/Experimental/MutableAssetUserData/Source/MutableAssetUserData/Public/MutableAssetUserDataModule.h
new file mode 100644
--- /dev/null
+++ b/Engine/Plugins/Experimental/MutableAssetUserData/Source/MutableAssetUserData/Public/MutableAssetUserDataModule.h
@@ -0,0 +1,13 @@
+// Copyright Epic Games, Inc. All Rights Reserved.
+
+#pragma once
+
+#include "Modules/ModuleInterface.h"
+
+
+class FMutableAssetUserDataModule : public IModuleInterface
+{
+	virtual void StartupModule() override;
+	virtual void ShutdownModule() override;
+};
+
