[Mutable] Exposed Asset User Data to BP.

#rnx
#Synced-CL 51537794

diff --git a/Engine/Plugins/Experimental/MutableAssetUserData/Source/MutableAssetUserData/Public/MutableAssetUserData.h b/Engine/Plugins/Experimental/MutableAssetUserData/Source/MutableAssetUserData/Public/MutableAssetUserData.h
--- a/Engine/Plugins/Experimental/MutableAssetUserData/Source/MutableAssetUserData/Public/MutableAssetUserData.h
+++ b/Engine/Plugins/Experimental/MutableAssetUserData/Source/MutableAssetUserData/Public/MutableAssetUserData.h
@@ -8,11 +8,11 @@
 
 class UAssetUserData;
 
-USTRUCT(meta = (DisplayName = "Asset User Data", PinColor = "#949494"))
+USTRUCT(BlueprintType, meta = (DisplayName = "Asset User Data", PinColor = "#949494"))
 struct FMutableAssetUserData
 {
 	GENERATED_BODY()
 
-	UPROPERTY(EditAnywhere, Category = AssetUserData)
+	UPROPERTY(BlueprintReadWrite, Category = AssetUserData)
 	TObjectPtr<UAssetUserData> AssetUserData;
 };
\ No newline at end of file
