[Mutable] Fix unable to save MutableAssetUserData default value due to pointing a private.

Since this type does not have constants simply make the property transient.

#Synced-CL 52129579

diff --git a/Engine/Plugins/Experimental/MutableAssetUserData/Source/MutableAssetUserData/Public/MutableAssetUserData.h b/Engine/Plugins/Experimental/MutableAssetUserData/Source/MutableAssetUserData/Public/MutableAssetUserData.h
--- a/Engine/Plugins/Experimental/MutableAssetUserData/Source/MutableAssetUserData/Public/MutableAssetUserData.h
+++ b/Engine/Plugins/Experimental/MutableAssetUserData/Source/MutableAssetUserData/Public/MutableAssetUserData.h
@@ -13,6 +13,6 @@ struct FMutableAssetUserData
 {
 	GENERATED_BODY()
 
-	UPROPERTY(BlueprintReadWrite, Category = AssetUserData)
+	UPROPERTY(BlueprintReadWrite, Transient, Category = AssetUserData)
 	TObjectPtr<UAssetUserData> AssetUserData;
 };
\ No newline at end of file
